Update the HandheldARBP template project file

diff --git a/Templates/TP_HandheldARBP/TP_HandheldARBP.uproject b/Templates/TP_HandheldARBP/TP_HandheldARBP.uproject
--- a/Templates/TP_HandheldARBP/TP_HandheldARBP.uproject
+++ b/Templates/TP_HandheldARBP/TP_HandheldARBP.uproject
@@ -1,5 +1,6 @@
 {
 	"FileVersion": 3,
+	"EngineAssociation": "",
 	"Category": "",
 	"Description": "",
 	"Plugins": [
@@ -12,7 +13,10 @@
 		},
 		{
 			"Name": "AppleARKit",
-			"Enabled": true
+			"Enabled": true,
+			"SupportedTargetPlatforms": [
+				"IOS"
+			]
 		}
 	],
 	"TargetPlatforms": [
