Added ARKit as a plugin dependency for the Handheld AR Template

diff --git a/Templates/TP_HandheldARBP/TP_HandheldARBP.uproject b/Templates/TP_HandheldARBP/TP_HandheldARBP.uproject
--- a/Templates/TP_HandheldARBP/TP_HandheldARBP.uproject
+++ b/Templates/TP_HandheldARBP/TP_HandheldARBP.uproject
@@ -6,10 +6,14 @@
 		{
 			"Name": "GoogleARCore",
 			"Enabled": true
+		},
+		{
+			"Name": "AppleARKit",
+			"Enabled": true
 		}
 	],
 	"TargetPlatforms": [
 		"Android",
 		"IOS"
 	]
-}
\ No newline at end of file
+}
