Adding ARKit directory to a whitelist for Android in the HandheldAR project

Also updating the .uproject for the HandheldAR project template, as it was warning on startup that it the project file was out of date.

diff --git a/Templates/TP_HandheldARBP/Config/Android/AndroidGame.ini b/Templates/TP_HandheldARBP/Config/Android/AndroidGame.ini
new file mode 100644
--- /dev/null
+++ b/Templates/TP_HandheldARBP/Config/Android/AndroidGame.ini
@@ -0,0 +1,2 @@
+[Staging]
++WhitelistDirectories=Engine/Plugins/Runtime/AR/Apple/AppleARKit
\ No newline at end of file
diff --git a/Templates/TP_HandheldARBP/TP_HandheldARBP.uproject b/Templates/TP_HandheldARBP/TP_HandheldARBP.uproject
--- a/Templates/TP_HandheldARBP/TP_HandheldARBP.uproject
+++ b/Templates/TP_HandheldARBP/TP_HandheldARBP.uproject
@@ -5,7 +5,10 @@
 	"Plugins": [
 		{
 			"Name": "GoogleARCore",
-			"Enabled": true
+			"Enabled": true,
+			"SupportedTargetPlatforms": [
+				"Android"
+			]
 		},
 		{
 			"Name": "AppleARKit",
